Made the proxy create function blueprint internal only

diff --git a/Engine/Plugins/Experimental/AppleVision/Source/AppleVision/Public/AppleVisionBlueprintProxy.h b/Engine/Plugins/Experimental/AppleVision/Source/AppleVision/Public/AppleVisionBlueprintProxy.h
--- a/Engine/Plugins/Experimental/AppleVision/Source/AppleVision/Public/AppleVisionBlueprintProxy.h
+++ b/Engine/Plugins/Experimental/AppleVision/Source/AppleVision/Public/AppleVisionBlueprintProxy.h
@@ -35,7 +35,7 @@ public:
 	 *
 	 * @param SourceImage the image to detect faces in
 	 */
-	UFUNCTION(BlueprintCallable, Meta=(DisplayName="Detect Faces"), Category="Apple Vision")
+	UFUNCTION(BlueprintCallable, Meta=(BlueprintInternalUseOnly="true", DisplayName="Detect Faces"), Category="Apple Vision")
 	static UAppleVisionDetectFacesAsyncTaskBlueprintProxy* CreateProxyObjectForDetectFaces(UTexture* SourceImage);
 
 	//~ Begin FTickableObject Interface
