Fixed mismatched copy constructors/assignment operators in FDetectedFeatureRegion, FDetectedFeatureRegion and FFaceDetectionResult.

diff --git a/Engine/Plugins/Experimental/AppleVision/Source/AppleVision/Public/AppleVisionTypes.h b/Engine/Plugins/Experimental/AppleVision/Source/AppleVision/Public/AppleVisionTypes.h
--- a/Engine/Plugins/Experimental/AppleVision/Source/AppleVision/Public/AppleVisionTypes.h
+++ b/Engine/Plugins/Experimental/AppleVision/Source/AppleVision/Public/AppleVisionTypes.h
@@ -40,13 +40,6 @@ public:
 	/** The set of points that encompass the detected feature area */
 	UPROPERTY(BlueprintReadOnly, Category="Apple Vision")
 	TArray<FVector2D> Points;
-
-	FORCEINLINE FDetectedFeatureRegion& operator=(const FDetectedFeatureRegion& Other)
-	{
-		Confidence = Other.Confidence;
-		Points = Other.Points;
-		return *this;
-	}
 };
 
 /** Features of a face that can be detected */
@@ -113,16 +106,6 @@ public:
 	/** The set of region features that were detected */
 	UPROPERTY(BlueprintReadOnly, Category="Apple Vision")
 	TArray<FDetectedFaceFeatureRegion> FeatureRegions;
-
-	FORCEINLINE FDetectedFace& operator=(const FDetectedFace& Other)
-	{
-		Confidence = Other.Confidence;
-		Points = Other.Points;
-		BoundingBox = Other.BoundingBox;
-		Features = Other.Features;
-		FeatureRegions = Other.FeatureRegions;
-		return *this;
-	}
 };
 
 /** The result of a face detection request with information about the detected faces */
@@ -137,10 +120,4 @@ public:
 	/** The set of faces that were detected */
 	UPROPERTY(BlueprintReadOnly, Category="Apple Vision")
 	TArray<FDetectedFace> DetectedFaces;
-
-	FORCEINLINE FFaceDetectionResult& operator=(const FFaceDetectionResult& Other)
-	{
-		DetectedFaces = Other.DetectedFaces;
-		return *this;
-	}
 };
